Fix C++ projects created from the Blank template producing BuildSettingsVersion warnings by default.

#rnx

diff --git a/Templates/TP_Blank/Source/TP_Blank.Target.cs b/Templates/TP_Blank/Source/TP_Blank.Target.cs
--- a/Templates/TP_Blank/Source/TP_Blank.Target.cs
+++ b/Templates/TP_Blank/Source/TP_Blank.Target.cs
@@ -8,7 +8,7 @@ public class TP_BlankTarget : TargetRules
 	public TP_BlankTarget( TargetInfo Target) : base(Target)
 	{
 		Type = TargetType.Game;
-
+		DefaultBuildSettings = BuildSettingsVersion.Latest;
 		ExtraModuleNames.AddRange( new string[] { "TP_Blank" } );
 	}
 }
diff --git a/Templates/TP_Blank/Source/TP_BlankEditor.Target.cs b/Templates/TP_Blank/Source/TP_BlankEditor.Target.cs
--- a/Templates/TP_Blank/Source/TP_BlankEditor.Target.cs
+++ b/Templates/TP_Blank/Source/TP_BlankEditor.Target.cs
@@ -8,7 +8,7 @@ public class TP_BlankEditorTarget : TargetRules
 	public TP_BlankEditorTarget( TargetInfo Target) : base(Target)
 	{
 		Type = TargetType.Editor;
-
+		DefaultBuildSettings = BuildSettingsVersion.Latest;
 		ExtraModuleNames.AddRange( new string[] { "TP_Blank" } );
 	}
 }
