Update templates to set IncludeOrderVersion = EngineIncludeOrderVersion.Latest;

#rnx

diff --git a/Templates/TP_Blank/Source/TP_Blank.Target.cs b/Templates/TP_Blank/Source/TP_Blank.Target.cs
--- a/Templates/TP_Blank/Source/TP_Blank.Target.cs
+++ b/Templates/TP_Blank/Source/TP_Blank.Target.cs
@@ -9,6 +9,7 @@ public class TP_BlankTarget : TargetRules
 	{
 		Type = TargetType.Game;
 		DefaultBuildSettings = BuildSettingsVersion.Latest;
-		ExtraModuleNames.AddRange( new string[] { "TP_Blank" } );
+		IncludeOrderVersion = EngineIncludeOrderVersion.Latest;
+		ExtraModuleNames.Add("TP_Blank");
 	}
 }
diff --git a/Templates/TP_Blank/Source/TP_BlankEditor.Target.cs b/Templates/TP_Blank/Source/TP_BlankEditor.Target.cs
--- a/Templates/TP_Blank/Source/TP_BlankEditor.Target.cs
+++ b/Templates/TP_Blank/Source/TP_BlankEditor.Target.cs
@@ -9,6 +9,7 @@ public class TP_BlankEditorTarget : TargetRules
 	{
 		Type = TargetType.Editor;
 		DefaultBuildSettings = BuildSettingsVersion.Latest;
-		ExtraModuleNames.AddRange( new string[] { "TP_Blank" } );
+		IncludeOrderVersion = EngineIncludeOrderVersion.Latest;
+		ExtraModuleNames.Add("TP_Blank");
 	}
 }
diff --git a/Templates/TP_FirstPerson/Source/TP_FirstPerson.Target.cs b/Templates/TP_FirstPerson/Source/TP_FirstPerson.Target.cs
--- a/Templates/TP_FirstPerson/Source/TP_FirstPerson.Target.cs
+++ b/Templates/TP_FirstPerson/Source/TP_FirstPerson.Target.cs
@@ -9,6 +9,7 @@ public class TP_FirstPersonTarget : TargetRules
 	{
 		Type = TargetType.Game;
 		DefaultBuildSettings = BuildSettingsVersion.Latest;
+		IncludeOrderVersion = EngineIncludeOrderVersion.Latest;
 		ExtraModuleNames.Add("TP_FirstPerson");
 	}
 }
diff --git a/Templates/TP_FirstPerson/Source/TP_FirstPersonEditor.Target.cs b/Templates/TP_FirstPerson/Source/TP_FirstPersonEditor.Target.cs
--- a/Templates/TP_FirstPerson/Source/TP_FirstPersonEditor.Target.cs
+++ b/Templates/TP_FirstPerson/Source/TP_FirstPersonEditor.Target.cs
@@ -9,6 +9,7 @@ public class TP_FirstPersonEditorTarget : TargetRules
 	{
 		Type = TargetType.Editor;
 		DefaultBuildSettings = BuildSettingsVersion.Latest;
+		IncludeOrderVersion = EngineIncludeOrderVersion.Latest;
 		ExtraModuleNames.Add("TP_FirstPerson");
 	}
 }
diff --git a/Templates/TP_Puzzle/Source/TP_Puzzle.Target.cs b/Templates/TP_Puzzle/Source/TP_Puzzle.Target.cs
--- a/Templates/TP_Puzzle/Source/TP_Puzzle.Target.cs
+++ b/Templates/TP_Puzzle/Source/TP_Puzzle.Target.cs
@@ -9,6 +9,7 @@ public class TP_PuzzleTarget : TargetRules
 	{
 		Type = TargetType.Game;
 		DefaultBuildSettings = BuildSettingsVersion.Latest;
+		IncludeOrderVersion = EngineIncludeOrderVersion.Latest;
 		ExtraModuleNames.Add("TP_Puzzle");
 	}
 }
diff --git a/Templates/TP_Puzzle/Source/TP_PuzzleEditor.Target.cs b/Templates/TP_Puzzle/Source/TP_PuzzleEditor.Target.cs
--- a/Templates/TP_Puzzle/Source/TP_PuzzleEditor.Target.cs
+++ b/Templates/TP_Puzzle/Source/TP_PuzzleEditor.Target.cs
@@ -9,6 +9,7 @@ public class TP_PuzzleEditorTarget : TargetRules
 	{
 		Type = TargetType.Editor;
 		DefaultBuildSettings = BuildSettingsVersion.Latest;
+		IncludeOrderVersion = EngineIncludeOrderVersion.Latest;
 		ExtraModuleNames.Add("TP_Puzzle");
 	}
 }
diff --git a/Templates/TP_ThirdPerson/Source/TP_ThirdPerson.Target.cs b/Templates/TP_ThirdPerson/Source/TP_ThirdPerson.Target.cs
--- a/Templates/TP_ThirdPerson/Source/TP_ThirdPerson.Target.cs
+++ b/Templates/TP_ThirdPerson/Source/TP_ThirdPerson.Target.cs
@@ -9,6 +9,7 @@ public class TP_ThirdPersonTarget : TargetRules
 	{
 		Type = TargetType.Game;
 		DefaultBuildSettings = BuildSettingsVersion.Latest;
+		IncludeOrderVersion = EngineIncludeOrderVersion.Latest;
 		ExtraModuleNames.Add("TP_ThirdPerson");
 	}
 }
diff --git a/Templates/TP_ThirdPerson/Source/TP_ThirdPersonEditor.Target.cs b/Templates/TP_ThirdPerson/Source/TP_ThirdPersonEditor.Target.cs
--- a/Templates/TP_ThirdPerson/Source/TP_ThirdPersonEditor.Target.cs
+++ b/Templates/TP_ThirdPerson/Source/TP_ThirdPersonEditor.Target.cs
@@ -9,6 +9,7 @@ public class TP_ThirdPersonEditorTarget : TargetRules
 	{
 		Type = TargetType.Editor;
 		DefaultBuildSettings = BuildSettingsVersion.Latest;
+		IncludeOrderVersion = EngineIncludeOrderVersion.Latest;
 		ExtraModuleNames.Add("TP_ThirdPerson");
 	}
 }
diff --git a/Templates/TP_TopDown/Source/TP_TopDown.Target.cs b/Templates/TP_TopDown/Source/TP_TopDown.Target.cs
--- a/Templates/TP_TopDown/Source/TP_TopDown.Target.cs
+++ b/Templates/TP_TopDown/Source/TP_TopDown.Target.cs
@@ -9,6 +9,7 @@ public class TP_TopDownTarget : TargetRules
 	{
 		Type = TargetType.Game;
 		DefaultBuildSettings = BuildSettingsVersion.Latest;
+		IncludeOrderVersion = EngineIncludeOrderVersion.Latest;
 		ExtraModuleNames.Add("TP_TopDown");
 	}
 }
diff --git a/Templates/TP_TopDown/Source/TP_TopDownEditor.Target.cs b/Templates/TP_TopDown/Source/TP_TopDownEditor.Target.cs
--- a/Templates/TP_TopDown/Source/TP_TopDownEditor.Target.cs
+++ b/Templates/TP_TopDown/Source/TP_TopDownEditor.Target.cs
@@ -9,6 +9,7 @@ public class TP_TopDownEditorTarget : TargetRules
 	{
 		Type = TargetType.Editor;
 		DefaultBuildSettings = BuildSettingsVersion.Latest;
+		IncludeOrderVersion = EngineIncludeOrderVersion.Latest;
 		ExtraModuleNames.Add("TP_TopDown");
 	}
 }
diff --git a/Templates/TP_VehicleAdv/Source/TP_VehicleAdv.Target.cs b/Templates/TP_VehicleAdv/Source/TP_VehicleAdv.Target.cs
--- a/Templates/TP_VehicleAdv/Source/TP_VehicleAdv.Target.cs
+++ b/Templates/TP_VehicleAdv/Source/TP_VehicleAdv.Target.cs
@@ -9,6 +9,7 @@ public class TP_VehicleAdvTarget : TargetRules
 	{
 		Type = TargetType.Game;
 		DefaultBuildSettings = BuildSettingsVersion.Latest;
+		IncludeOrderVersion = EngineIncludeOrderVersion.Latest;
 		ExtraModuleNames.Add("TP_VehicleAdv");
 	}
 }
diff --git a/Templates/TP_VehicleAdv/Source/TP_VehicleAdvEditor.Target.cs b/Templates/TP_VehicleAdv/Source/TP_VehicleAdvEditor.Target.cs
--- a/Templates/TP_VehicleAdv/Source/TP_VehicleAdvEditor.Target.cs
+++ b/Templates/TP_VehicleAdv/Source/TP_VehicleAdvEditor.Target.cs
@@ -9,6 +9,7 @@ public class TP_VehicleAdvEditorTarget : TargetRules
 	{
 		Type = TargetType.Editor;
 		DefaultBuildSettings = BuildSettingsVersion.Latest;
+		IncludeOrderVersion = EngineIncludeOrderVersion.Latest;
 		ExtraModuleNames.Add("TP_VehicleAdv");
 	}
 }
