Add the EnhancedInput module by default to the TP_Blank.build.cs file. This module is going to be needed by anyone who wants to use Enhanced Input in their project, so adding it by default prevents headaches when creating a new project.

diff --git a/Templates/TP_Blank/Source/TP_Blank/TP_Blank.Build.cs b/Templates/TP_Blank/Source/TP_Blank/TP_Blank.Build.cs
--- a/Templates/TP_Blank/Source/TP_Blank/TP_Blank.Build.cs
+++ b/Templates/TP_Blank/Source/TP_Blank/TP_Blank.Build.cs
@@ -8,7 +8,7 @@ public class TP_Blank : ModuleRules
 	{
 		PCHUsage = PCHUsageMode.UseExplicitOrSharedPCHs;
 	
-		PublicDependencyModuleNames.AddRange(new string[] { "Core", "CoreUObject", "Engine", "InputCore" });
+		PublicDependencyModuleNames.AddRange(new string[] { "Core", "CoreUObject", "Engine", "InputCore", "EnhancedInput" });
 
 		PrivateDependencyModuleNames.AddRange(new string[] {  });
 
